Renamed Freezing_MaxFieldAlignment -> Freezing_MaxFieldAlignment64 in DataDrivenPlatformInfo.ini to fix serialization issues

diff --git a/Engine/Config/HoloLens/DataDrivenPlatformInfo.ini b/Engine/Config/HoloLens/DataDrivenPlatformInfo.ini
--- a/Engine/Config/HoloLens/DataDrivenPlatformInfo.ini
+++ b/Engine/Config/HoloLens/DataDrivenPlatformInfo.ini
@@ -1,6 +1,6 @@
 [DataDrivenPlatformInfo]
 bIsConfidential=false
-Freezing_MaxFieldAlignment=8
+Freezing_MaxFieldAlignment64=8
 Freezing_bAlignBases=true
 
 NormalIconPath=Launcher/Windows/Platform_Windows_24x
