PR #8636: Add DataDriven ShaderPlatform entries for HoloLens (Contributed by [REDACTED_PERSON])

diff --git a/Engine/Config/HoloLens/DataDrivenPlatformInfo.ini b/Engine/Config/HoloLens/DataDrivenPlatformInfo.ini
--- a/Engine/Config/HoloLens/DataDrivenPlatformInfo.ini
+++ b/Engine/Config/HoloLens/DataDrivenPlatformInfo.ini
@@ -26,3 +26,10 @@ bUsesHostCompiler=true
 bCanUseCrashReporter=false
 bUATClosesAfterLaunch=false
 PlatformGroupName=Mobile
+
+[ShaderPlatform PCD3D_ES3_1]
+Language=D3D
+MaxFeatureLevel=ES3_1
+bIsMobile=true
+bSupportsMobileMultiView=true
+bSupportsInstancedStereo=true
