Back out Hololens DataDrivenPlatformInfo.ini change that conflicts with Windows

diff --git a/Engine/Config/HoloLens/DataDrivenPlatformInfo.ini b/Engine/Config/HoloLens/DataDrivenPlatformInfo.ini
--- a/Engine/Config/HoloLens/DataDrivenPlatformInfo.ini
+++ b/Engine/Config/HoloLens/DataDrivenPlatformInfo.ini
@@ -26,10 +26,3 @@ bUsesHostCompiler=true
 bCanUseCrashReporter=false
 bUATClosesAfterLaunch=false
 PlatformGroupName=Mobile
-
-[ShaderPlatform PCD3D_ES3_1]
-Language=D3D
-MaxFeatureLevel=ES3_1
-bIsMobile=true
-bSupportsMobileMultiView=true
-bSupportsInstancedStereo=true
