enable Oodle for HoloLens
was disabled temporarily when Oodle didn't have a HoloLens build, forgot to remove the temporary setting when that was fixed

diff --git a/Engine/Config/HoloLens/DataDrivenPlatformInfo.ini b/Engine/Config/HoloLens/DataDrivenPlatformInfo.ini
--- a/Engine/Config/HoloLens/DataDrivenPlatformInfo.ini
+++ b/Engine/Config/HoloLens/DataDrivenPlatformInfo.ini
@@ -11,8 +11,6 @@ bSupportsGamepad=true
 bCanChangeGamepadType=true
 bSupportsTouch=false
 GlobalIdentifier=0DF604E112E4445280BEE1C74EB934B1
-; Oodle is not supported on this platform so force something else to be used (Zlib or None) :
-HardwareCompressionFormat=Zlib
 
 NormalIconPath=Launcher/Windows/Platform_Windows_24x
 LargeIconPath=Launcher/Windows/Platform_Windows_128x
