Hololens + D3D12: Disable GPU crash tracking data

D3D12 uses WriteBufferImmediate to write breadcrumbs for GPU crash tracking. Hololens claims support ID3D12GraphicsCommandList2, but inexplicably, does not support WriteBufferImmediate (the only additional method exposed for GraphicsCommandList2). This is probably because the underlying QCOM implementation does not support the command.

We just disable the entire GPU crash tracking setup on Hololens, as it is dependent on the breadcrumbs.

diff --git a/Engine/Config/HoloLens/HoloLensEngine.ini b/Engine/Config/HoloLens/HoloLensEngine.ini
--- a/Engine/Config/HoloLens/HoloLensEngine.ini
+++ b/Engine/Config/HoloLens/HoloLensEngine.ini
@@ -26,6 +26,9 @@ fx.Niagara.QualityLevel.Max=0
 ; Enable IPv6 by default on Hololens
 net.DisableIPv6=0
 
+; Hololens doesn't support WriteBufferImmediate, which we need for D3D12 GPU breadcrumbs
+r.gpucrash.collectionenable=0
+
 [/Script/HoloLensPlatformEditor.HoloLensTargetSettings]
 Windows10SDKVersion=10.0.18362.0
 AudioMaxChannels=23
